Deprecating AES ECB packet handler component
Updated AES GCM packet handler to no longer be in beta

diff --git a/Engine/Config/BaseEngine.ini b/Engine/Config/BaseEngine.ini
--- a/Engine/Config/BaseEngine.ini
+++ b/Engine/Config/BaseEngine.ini
@@ -2904,7 +2904,7 @@ PktIncomingLagMax=200
 
 [PacketHandlerComponents]
 ; Options can be set in this section of DefaultEngine.ini to enable different types of network packet encruption plugins
-; EncryptionComponent=AESHandlerComponent
+; EncryptionComponent=AESGCMHandlerComponent
 
 [/Script/GameplayDebugger.GameplayDebuggingControllerComponent]
 CategoryZeroBind=(Key=NumPadZero,bShift=False,bCtrl=False,bAlt=False,bCmd=False)
diff --git a/Engine/Plugins/Runtime/PacketHandlers/AESGCMHandlerComponent/AESGCMHandlerComponent.uplugin b/Engine/Plugins/Runtime/PacketHandlers/AESGCMHandlerComponent/AESGCMHandlerComponent.uplugin
--- a/Engine/Plugins/Runtime/PacketHandlers/AESGCMHandlerComponent/AESGCMHandlerComponent.uplugin
+++ b/Engine/Plugins/Runtime/PacketHandlers/AESGCMHandlerComponent/AESGCMHandlerComponent.uplugin
@@ -12,7 +12,7 @@
 	"SupportURL" : "",
 	"EnabledByDefault" : false,
 	"CanContainContent" : false,
-	"IsBetaVersion" : true,
+	"IsBetaVersion" : false,
 	"Installed" : false,
 	"Modules" :
 	[
diff --git a/Engine/Plugins/Runtime/PacketHandlers/AESHandlerComponent/Source/Public/AESHandlerComponent.h b/Engine/Plugins/Runtime/PacketHandlers/AESHandlerComponent/Source/Public/AESHandlerComponent.h
--- a/Engine/Plugins/Runtime/PacketHandlers/AESHandlerComponent/Source/Public/AESHandlerComponent.h
+++ b/Engine/Plugins/Runtime/PacketHandlers/AESHandlerComponent/Source/Public/AESHandlerComponent.h
@@ -10,7 +10,7 @@
 /*
 * AES256 ECB block encryption component.
 */
-class AESHANDLERCOMPONENT_API FAESHandlerComponent : public FEncryptionComponent
+class UE_DEPRECATED(5.1, "No longer supported, please use FAESGCMHandlerComponent instead.") AESHANDLERCOMPONENT_API FAESHandlerComponent : public FEncryptionComponent
 {
 public:
 	/**
@@ -60,7 +60,7 @@ private:
 /**
  * The public interface to this module.
  */
-class FAESHandlerComponentModule : public FPacketHandlerComponentModuleInterface
+class UE_DEPRECATED(5.1, "No longer supported, please use FAESGCMHandlerComponentModule instead.") FAESHandlerComponentModule : public FPacketHandlerComponentModuleInterface
 {
 public:
 	/* Creates an instance of this component */
