Fix deprecation warnings

diff --git a/Engine/Plugins/Runtime/PacketHandlers/AESHandlerComponent/Source/Private/AESHandlerComponent.cpp b/Engine/Plugins/Runtime/PacketHandlers/AESHandlerComponent/Source/Private/AESHandlerComponent.cpp
--- a/Engine/Plugins/Runtime/PacketHandlers/AESHandlerComponent/Source/Private/AESHandlerComponent.cpp
+++ b/Engine/Plugins/Runtime/PacketHandlers/AESHandlerComponent/Source/Private/AESHandlerComponent.cpp
@@ -2,6 +2,8 @@
 
 #include "AESHandlerComponent.h"
 
+PRAGMA_DISABLE_DEPRECATION_WARNINGS
+
 IMPLEMENT_MODULE( FAESHandlerComponentModule, AESHandlerComponent )
 
 const int32 FAESHandlerComponent::KeySizeInBytes;
@@ -233,3 +235,5 @@ void FAESHandlerComponent::CountBytes(FArchive& Ar) const
 	Key.CountBytes(Ar);
 	Ciphertext.CountBytes(Ar);
 }
+
+PRAGMA_ENABLE_DEPRECATION_WARNINGS
\ No newline at end of file
